Move the FlushTick from OnlineServicesEOSGS to OnlineServicesEpicCommon to make sure other online components which rely on EOS can also tick and complete the operations during shutdown.

diff --git a/Engine/Plugins/Online/OnlineServicesEOSGS/Source/Private/Online/OnlineServicesEOSGS.cpp b/Engine/Plugins/Online/OnlineServicesEOSGS/Source/Private/Online/OnlineServicesEOSGS.cpp
--- a/Engine/Plugins/Online/OnlineServicesEOSGS/Source/Private/Online/OnlineServicesEOSGS.cpp
+++ b/Engine/Plugins/Online/OnlineServicesEOSGS/Source/Private/Online/OnlineServicesEOSGS.cpp
@@ -65,13 +65,6 @@ void FOnlineServicesEOSGS::Destroy()
 #endif
 }
 
-void FOnlineServicesEOSGS::FlushTick(float DeltaSeconds)
-{
-	Super::FlushTick(DeltaSeconds);
-
-	GetEOSPlatformHandle()->Tick();
-}
-
 void FOnlineServicesEOSGS::WarnIfEncryptionKeyMissing(const FString& InterfaceName) const
 {
 	if (IEOSSDKManager* Manager = IEOSSDKManager::Get())
diff --git a/Engine/Plugins/Online/OnlineServicesEOSGS/Source/Public/Online/OnlineServicesEOSGS.h b/Engine/Plugins/Online/OnlineServicesEOSGS/Source/Public/Online/OnlineServicesEOSGS.h
--- a/Engine/Plugins/Online/OnlineServicesEOSGS/Source/Public/Online/OnlineServicesEOSGS.h
+++ b/Engine/Plugins/Online/OnlineServicesEOSGS/Source/Public/Online/OnlineServicesEOSGS.h
@@ -42,8 +42,6 @@ public:
 	static const TCHAR* GetServiceConfigNameStatic() { return TEXT("EOS"); }
 
 protected:
-	virtual void FlushTick(float DeltaSeconds) override;
-
 	void WarnIfEncryptionKeyMissing(const FString& InterfaceName) const;
 
 #if WITH_ENGINE
diff --git a/Engine/Plugins/Online/OnlineServicesEpicCommon/Source/OnlineServicesEpicCommon/Private/Online/OnlineServicesEpicCommon.cpp b/Engine/Plugins/Online/OnlineServicesEpicCommon/Source/OnlineServicesEpicCommon/Private/Online/OnlineServicesEpicCommon.cpp
--- a/Engine/Plugins/Online/OnlineServicesEpicCommon/Source/OnlineServicesEpicCommon/Private/Online/OnlineServicesEpicCommon.cpp
+++ b/Engine/Plugins/Online/OnlineServicesEpicCommon/Source/OnlineServicesEpicCommon/Private/Online/OnlineServicesEpicCommon.cpp
@@ -69,4 +69,11 @@ void FOnlineServicesEpicCommon::RemoveEOSSDKFastTick(FOnlineAsyncOp& InAsyncOp)
 	}
 }
 
+void FOnlineServicesEpicCommon::FlushTick(float DeltaSeconds)
+{
+	Super::FlushTick(DeltaSeconds);
+
+	GetEOSPlatformHandle()->Tick();
+}
+
 /* UE::Online */ }
diff --git a/Engine/Plugins/Online/OnlineServicesEpicCommon/Source/OnlineServicesEpicCommon/Public/Online/OnlineServicesEpicCommon.h b/Engine/Plugins/Online/OnlineServicesEpicCommon/Source/OnlineServicesEpicCommon/Public/Online/OnlineServicesEpicCommon.h
--- a/Engine/Plugins/Online/OnlineServicesEpicCommon/Source/OnlineServicesEpicCommon/Public/Online/OnlineServicesEpicCommon.h
+++ b/Engine/Plugins/Online/OnlineServicesEpicCommon/Source/OnlineServicesEpicCommon/Public/Online/OnlineServicesEpicCommon.h
@@ -27,6 +27,8 @@ public:
 	void RemoveEOSSDKFastTick(FOnlineAsyncOp& InAsyncOp);
 
 protected:
+	virtual void FlushTick(float DeltaSeconds) override;
+
 	IEOSPlatformHandlePtr EOSPlatformHandle;
 
 	bool bEnableAsyncOpFastTick = true;
