Removed files and usage for redundant class FOnlineServicesEOSGSPlatformFactory

diff --git a/Engine/Plugins/Online/OnlineServicesEOSGS/Source/Private/Online/OnlineServicesEOSGS.cpp b/Engine/Plugins/Online/OnlineServicesEOSGS/Source/Private/Online/OnlineServicesEOSGS.cpp
--- a/Engine/Plugins/Online/OnlineServicesEOSGS/Source/Private/Online/OnlineServicesEOSGS.cpp
+++ b/Engine/Plugins/Online/OnlineServicesEOSGS/Source/Private/Online/OnlineServicesEOSGS.cpp
@@ -7,7 +7,6 @@
 #include "Online/LeaderboardsEOSGS.h"
 #include "Online/LobbiesEOSGS.h"
 #include "Online/OnlineIdEOSGS.h"
-#include "Online/OnlineServicesEOSGSPlatformFactory.h"
 #include "Online/StatsEOSGS.h"
 #include "Online/SessionsEOSGS.h"
 #include "Online/TitleFileEOSGS.h"
diff --git a/Engine/Plugins/Online/OnlineServicesEOSGS/Source/Private/Online/OnlineServicesEOSGSModule.cpp b/Engine/Plugins/Online/OnlineServicesEOSGS/Source/Private/Online/OnlineServicesEOSGSModule.cpp
--- a/Engine/Plugins/Online/OnlineServicesEOSGS/Source/Private/Online/OnlineServicesEOSGSModule.cpp
+++ b/Engine/Plugins/Online/OnlineServicesEOSGS/Source/Private/Online/OnlineServicesEOSGSModule.cpp
@@ -7,7 +7,6 @@
 #include "Online/OnlineIdEOSGS.h"
 #include "Online/OnlineServicesEOSGS.h"
 #include "Online/OnlineServicesRegistry.h"
-#include "Online/OnlineServicesEOSGSPlatformFactory.h"
 #include "Online/SessionsEOSGS.h"
 
 namespace UE::Online
@@ -58,9 +57,6 @@ void FOnlineServicesEOSGSModule::StartupModule()
 
 	static FOnlineSessionInviteIdRegistryEOSGS SessionInviteIdRegistry(EOnlineServices::Epic);
 	FOnlineIdRegistryRegistry::Get().RegisterSessionInviteIdRegistry(EOnlineServices::Epic, &SessionInviteIdRegistry, GetRegistryPriority());
-
-	// Initialize the platform factory on startup.  This is necessary for the SDK to bind to rendering and input very early.
-	FOnlineServicesEOSGSPlatformFactory::Get();
 }
 
 void FOnlineServicesEOSGSModule::ShutdownModule()
@@ -69,8 +65,6 @@ void FOnlineServicesEOSGSModule::ShutdownModule()
 	FOnlineIdRegistryRegistry::Get().UnregisterAccountIdRegistry(EOnlineServices::Epic, GetRegistryPriority());
 	FOnlineIdRegistryRegistry::Get().UnregisterSessionIdRegistry(EOnlineServices::Epic, GetRegistryPriority());
 	FOnlineIdRegistryRegistry::Get().UnregisterSessionInviteIdRegistry(EOnlineServices::Epic, GetRegistryPriority());
-
-	UE::Online::FOnlineServicesEOSGSPlatformFactory::TearDown();
 }
 
 /* UE::Online */ }
diff --git a/Engine/Plugins/Online/OnlineServicesEOSGS/Source/Private/Online/OnlineServicesEOSGSPlatformFactory.cpp b/Engine/Plugins/Online/OnlineServicesEOSGS/Source/Private/Online/OnlineServicesEOSGSPlatformFactory.cpp
deleted file mode 100644
diff --git a/Engine/Plugins/Online/OnlineServicesEOSGS/Source/Private/Online/OnlineServicesEOSGSPlatformFactory.h b/Engine/Plugins/Online/OnlineServicesEOSGS/Source/Private/Online/OnlineServicesEOSGSPlatformFactory.h
deleted file mode 100644
diff --git a/Engine/Plugins/Online/OnlineServicesEpicCommon/Source/OnlineServicesEpicCommon/Private/Online/OnlineServicesEpicCommonPlatformFactory.h b/Engine/Plugins/Online/OnlineServicesEpicCommon/Source/OnlineServicesEpicCommon/Private/Online/OnlineServicesEpicCommonPlatformFactory.h
new file mode 100644
--- /dev/null
+++ b/Engine/Plugins/Online/OnlineServicesEpicCommon/Source/OnlineServicesEpicCommon/Private/Online/OnlineServicesEpicCommonPlatformFactory.h
@@ -0,0 +1,47 @@
+// Copyright Epic Games, Inc. All Rights Reserved.
+
+#pragma once
+
+#include "Templates/SharedPointer.h"
+
+#include "IEOSSDKManager.h"
+
+using IEOSPlatformHandlePtr = TSharedPtr<class IEOSPlatformHandle>;
+class FLazySingleton;
+
+namespace UE::Online
+{
+
+	/** Factory class to create EOS Platforms for online services */
+	class FOnlineServicesEpicCommonPlatformFactory
+	{
+	public:
+		/**
+		 * Get the platform factory
+		 * @return the platform factory singleton
+		 */
+		static FOnlineServicesEpicCommonPlatformFactory& Get();
+		/**
+		 * Tear down the singleton instance. This only cleans up the singleton and has no impact on any platform handles created by this (aside from DefaultEOSPlatformHandle's ref count decreasing).
+		 */
+		static void TearDown();
+
+		/**
+		 * Create a new platform instance for the given Instance/Config pair.
+		 * Note If InstanceConfigName = NAME_None, this will attempt to resolve a config to use from various sources, including OnlineServices config, and EOSSDKManager cached configs.
+		 * @param InstanceName the "instance", typically a WorldContextHandle
+		 * @param InstanceConfigName the named EOS platform config to use
+		 * @return a new platform, or null on failure
+		 */
+		IEOSPlatformHandlePtr CreatePlatform(FName InstanceName, FName InstanceConfigName);
+
+	private:
+		FOnlineServicesEpicCommonPlatformFactory() = default;
+
+		FString LoadEOSPlatformConfig(IEOSSDKManager* const SDKManager);
+
+		friend FLazySingleton;
+	};
+
+	/* UE::Online */
+}
diff --git a/Engine/Plugins/Online/OnlineServicesEpicCommon/Source/OnlineServicesEpicCommon/Public/Online/OnlineServicesEpicCommonPlatformFactory.h b/Engine/Plugins/Online/OnlineServicesEpicCommon/Source/OnlineServicesEpicCommon/Public/Online/OnlineServicesEpicCommonPlatformFactory.h
deleted file mode 100644
