Added logic to process sandboxid and deploymentid overrides for EOS platform creation in OnlineServices

diff --git a/Engine/Plugins/Online/EOSShared/Source/EOSShared/Private/EOSSDKManager.cpp b/Engine/Plugins/Online/EOSShared/Source/EOSShared/Private/EOSSDKManager.cpp
--- a/Engine/Plugins/Online/EOSShared/Source/EOSShared/Private/EOSSDKManager.cpp
+++ b/Engine/Plugins/Online/EOSShared/Source/EOSShared/Private/EOSSDKManager.cpp
@@ -392,6 +392,30 @@ const FEOSSDKPlatformConfig* FEOSSDKManager::GetPlatformConfig(const FString& Pl
 	GConfig->GetInt(*SectionName, TEXT("TickBudgetInMilliseconds"), PlatformConfig->TickBudgetInMilliseconds, GEngineIni);
 	GConfig->GetArray(*SectionName, TEXT("OptionalConfig"), PlatformConfig->OptionalConfig, GEngineIni);
 
+	// After we have loaded the platform config, we'll check to see if there is any command-line overrides present
+
+	FString SandboxIdOverride;
+	// Get the -epicsandboxid argument. This generally comes from EGS.
+	bool bHasSandboxIdOverride = FParse::Value(FCommandLine::Get(), TEXT("EpicSandboxId="), SandboxIdOverride);
+	// Prefer -EpicSandboxIdOverride over previous.
+	bHasSandboxIdOverride |= FParse::Value(FCommandLine::Get(), TEXT("EpicSandboxIdOverride="), SandboxIdOverride);
+
+	if (bHasSandboxIdOverride)
+	{
+		PlatformConfig->SandboxId = SandboxIdOverride;
+	}
+
+	FString DeploymentIdOverride;
+	// Get the -epicdeploymentid argument. This generally comes from EGS.
+	bool bHasDeploymentIdOverride = FParse::Value(FCommandLine::Get(), TEXT("EpicDeploymentId="), DeploymentIdOverride);
+	// Prefer -EpicDeploymentIdOverride over previous.
+	bHasDeploymentIdOverride |= FParse::Value(FCommandLine::Get(), TEXT("EpicDeploymentIdOverride="), DeploymentIdOverride);
+
+	if (bHasDeploymentIdOverride)
+	{
+		PlatformConfig->DeploymentId = DeploymentIdOverride;
+	}
+
 	UE_LOG(LogEOSShared, Verbose, TEXT("%hs Loaded platform config: %s"), __FUNCTION__, *PlatformConfigName);
 	return PlatformConfig;
 }
diff --git a/Engine/Plugins/Online/OnlineServicesEpicCommon/Source/OnlineServicesEpicCommon/Private/Online/OnlineServicesEpicCommonPlatformFactory.cpp b/Engine/Plugins/Online/OnlineServicesEpicCommon/Source/OnlineServicesEpicCommon/Private/Online/OnlineServicesEpicCommonPlatformFactory.cpp
--- a/Engine/Plugins/Online/OnlineServicesEpicCommon/Source/OnlineServicesEpicCommon/Private/Online/OnlineServicesEpicCommonPlatformFactory.cpp
+++ b/Engine/Plugins/Online/OnlineServicesEpicCommon/Source/OnlineServicesEpicCommon/Private/Online/OnlineServicesEpicCommonPlatformFactory.cpp
@@ -6,6 +6,7 @@
 #include "Online/OnlineExecHandler.h"
 #include "Online/OnlineServicesEpicCommon.h"
 
+#include "Misc/CommandLine.h"
 #include "Misc/ConfigCacheIni.h"
 #include "Misc/Fork.h"
 #include "Misc/LazySingleton.h"
@@ -136,6 +137,30 @@ FString FOnlineServicesEpicCommonPlatformFactory::LoadEOSPlatformConfig(IEOSSDKM
 		PlatformConfig.bWindowsEnableOverlayOpenGL = true;
 	}
 
+	// After we have loaded the platform config, we'll check to see if there is any command-line overrides present
+
+	FString SandboxIdOverride;
+	// Get the -epicsandboxid argument. This generally comes from EGS.
+	bool bHasSandboxIdOverride = FParse::Value(FCommandLine::Get(), TEXT("EpicSandboxId="), SandboxIdOverride);
+	// Prefer -EpicSandboxIdOverride over previous.
+	bHasSandboxIdOverride |= FParse::Value(FCommandLine::Get(), TEXT("EpicSandboxIdOverride="), SandboxIdOverride);
+
+	if (bHasSandboxIdOverride)
+	{
+		PlatformConfig.SandboxId = SandboxIdOverride;
+	}
+
+	FString DeploymentIdOverride;
+	// Get the -epicdeploymentid argument. This generally comes from EGS.
+	bool bHasDeploymentIdOverride = FParse::Value(FCommandLine::Get(), TEXT("EpicDeploymentId="), DeploymentIdOverride);
+	// Prefer -EpicDeploymentIdOverride over previous.
+	bHasDeploymentIdOverride |= FParse::Value(FCommandLine::Get(), TEXT("EpicDeploymentIdOverride="), DeploymentIdOverride);
+
+	if (bHasDeploymentIdOverride)
+	{
+		PlatformConfig.DeploymentId = DeploymentIdOverride;
+	}
+
 	if (!SDKManager->AddPlatformConfig(PlatformConfig))
 	{
 		return FString();
