Created UserInfoEpicAccount to cover errors in functionality UserInfoEOS has when using split plugins

diff --git a/Engine/Plugins/Online/OnlineServicesEOS/Source/Private/Online/OnlineServicesEOS.cpp b/Engine/Plugins/Online/OnlineServicesEOS/Source/Private/Online/OnlineServicesEOS.cpp
--- a/Engine/Plugins/Online/OnlineServicesEOS/Source/Private/Online/OnlineServicesEOS.cpp
+++ b/Engine/Plugins/Online/OnlineServicesEOS/Source/Private/Online/OnlineServicesEOS.cpp
@@ -32,51 +32,4 @@ void FOnlineServicesEOS::RegisterComponents()
 	Super::RegisterComponents();
 }
 
-EOS_OnlinePlatformType EOnlinePlatformType_To_EOS_OnlinePlatformType(const EOnlinePlatformType& InType)
-{
-	switch (InType)
-	{
-	case EOnlinePlatformType::Epic:		return EOS_OPT_Epic;
-	case EOnlinePlatformType::Steam:	return EOS_OPT_Steam;
-	case EOnlinePlatformType::PSN:		return EOS_OPT_PSN;
-	case EOnlinePlatformType::Nintendo:	return EOS_OPT_Nintendo;
-	case EOnlinePlatformType::XBL:		return EOS_OPT_XBL;
-	default:							checkNoEntry(); // Intentional fall-through
-	case EOnlinePlatformType::Unknown:	return EOS_OPT_Unknown;
-	}
-}
-
-EOnlinePlatformType EOS_OnlinePlatformType_To_EOnlinePlatformType(const EOS_OnlinePlatformType& InType)
-{
-	if (InType == EOS_OPT_Epic)
-	{
-		return EOnlinePlatformType::Epic;
-	}
-	else if (InType == EOS_OPT_Steam)
-	{
-		return EOnlinePlatformType::Steam;
-	}
-	else if (InType == EOS_OPT_PSN)
-	{
-		return EOnlinePlatformType::PSN;
-	}
-	else if (InType == EOS_OPT_Nintendo)
-	{
-		return EOnlinePlatformType::Nintendo;
-	}
-	else if (InType == EOS_OPT_XBL)
-	{
-		return EOnlinePlatformType::XBL;
-	}
-	else if (InType == EOS_OPT_Unknown)
-	{
-		return EOnlinePlatformType::Unknown;
-	}
-	else
-	{
-		checkNoEntry();
-		return EOnlinePlatformType::Unknown;
-	}
-}
-
 /* UE::Online */ }
diff --git a/Engine/Plugins/Online/OnlineServicesEOS/Source/Public/Online/OnlineServicesEOS.h b/Engine/Plugins/Online/OnlineServicesEOS/Source/Public/Online/OnlineServicesEOS.h
--- a/Engine/Plugins/Online/OnlineServicesEOS/Source/Public/Online/OnlineServicesEOS.h
+++ b/Engine/Plugins/Online/OnlineServicesEOS/Source/Public/Online/OnlineServicesEOS.h
@@ -28,7 +28,4 @@ public:
 	ONLINESERVICESEOS_API virtual void RegisterComponents() override;
 };
 
-ONLINESERVICESEOS_API EOS_OnlinePlatformType EOnlinePlatformType_To_EOS_OnlinePlatformType(const EOnlinePlatformType& InType);
-ONLINESERVICESEOS_API EOnlinePlatformType EOS_OnlinePlatformType_To_EOnlinePlatformType(const EOS_OnlinePlatformType& InType);
-
 /* UE::Online */ }
diff --git a/Engine/Plugins/Online/OnlineServicesEpicCommon/Source/OnlineServicesEpicCommon/Private/Online/OnlineServicesEpicCommon.cpp b/Engine/Plugins/Online/OnlineServicesEpicCommon/Source/OnlineServicesEpicCommon/Private/Online/OnlineServicesEpicCommon.cpp
--- a/Engine/Plugins/Online/OnlineServicesEpicCommon/Source/OnlineServicesEpicCommon/Private/Online/OnlineServicesEpicCommon.cpp
+++ b/Engine/Plugins/Online/OnlineServicesEpicCommon/Source/OnlineServicesEpicCommon/Private/Online/OnlineServicesEpicCommon.cpp
@@ -97,4 +97,51 @@ void FOnlineServicesEpicCommon::FlushTick(float DeltaSeconds)
 	GetEOSPlatformHandle()->Tick();
 }
 
+EOS_OnlinePlatformType EOnlinePlatformType_To_EOS_OnlinePlatformType(const EOnlinePlatformType& InType)
+{
+	switch (InType)
+	{
+	case EOnlinePlatformType::Epic:		return EOS_OPT_Epic;
+	case EOnlinePlatformType::Steam:	return EOS_OPT_Steam;
+	case EOnlinePlatformType::PSN:		return EOS_OPT_PSN;
+	case EOnlinePlatformType::Nintendo:	return EOS_OPT_Nintendo;
+	case EOnlinePlatformType::XBL:		return EOS_OPT_XBL;
+	default:							checkNoEntry(); // Intentional fall-through
+	case EOnlinePlatformType::Unknown:	return EOS_OPT_Unknown;
+	}
+}
+
+EOnlinePlatformType EOS_OnlinePlatformType_To_EOnlinePlatformType(const EOS_OnlinePlatformType& InType)
+{
+	if (InType == EOS_OPT_Epic)
+	{
+		return EOnlinePlatformType::Epic;
+	}
+	else if (InType == EOS_OPT_Steam)
+	{
+		return EOnlinePlatformType::Steam;
+	}
+	else if (InType == EOS_OPT_PSN)
+	{
+		return EOnlinePlatformType::PSN;
+	}
+	else if (InType == EOS_OPT_Nintendo)
+	{
+		return EOnlinePlatformType::Nintendo;
+	}
+	else if (InType == EOS_OPT_XBL)
+	{
+		return EOnlinePlatformType::XBL;
+	}
+	else if (InType == EOS_OPT_Unknown)
+	{
+		return EOnlinePlatformType::Unknown;
+	}
+	else
+	{
+		checkNoEntry();
+		return EOnlinePlatformType::Unknown;
+	}
+}
+
 /* UE::Online */ }
diff --git a/Engine/Plugins/Online/OnlineServicesEpicCommon/Source/OnlineServicesEpicCommon/Public/Online/OnlineServicesEpicCommon.h b/Engine/Plugins/Online/OnlineServicesEpicCommon/Source/OnlineServicesEpicCommon/Public/Online/OnlineServicesEpicCommon.h
--- a/Engine/Plugins/Online/OnlineServicesEpicCommon/Source/OnlineServicesEpicCommon/Public/Online/OnlineServicesEpicCommon.h
+++ b/Engine/Plugins/Online/OnlineServicesEpicCommon/Source/OnlineServicesEpicCommon/Public/Online/OnlineServicesEpicCommon.h
@@ -4,6 +4,8 @@
 
 #include "Online/OnlineServicesCommon.h"
 
+#include "eos_common.h"
+
 #define UE_API ONLINESERVICESEPICCOMMON_API
 
 using IEOSPlatformHandlePtr = TSharedPtr<class IEOSPlatformHandle>;
@@ -38,6 +40,9 @@ protected:
 	bool bEnableAsyncOpFastTick = true;
 };
 
+UE_API EOS_OnlinePlatformType EOnlinePlatformType_To_EOS_OnlinePlatformType(const EOnlinePlatformType& InType);
+UE_API EOnlinePlatformType EOS_OnlinePlatformType_To_EOnlinePlatformType(const EOS_OnlinePlatformType& InType);
+
 /* UE::Online */ }
 
 #undef UE_API
