Increase ApiVersion for EOS_Auth_LoginOptions from 2 to 3

diff --git a/Engine/Plugins/Online/OnlineServicesEpicCommon/Source/OnlineServicesEpicCommon/Private/Online/EOSAuthLoginOptionsCommon.cpp b/Engine/Plugins/Online/OnlineServicesEpicCommon/Source/OnlineServicesEpicCommon/Private/Online/EOSAuthLoginOptionsCommon.cpp
--- a/Engine/Plugins/Online/OnlineServicesEpicCommon/Source/OnlineServicesEpicCommon/Private/Online/EOSAuthLoginOptionsCommon.cpp
+++ b/Engine/Plugins/Online/OnlineServicesEpicCommon/Source/OnlineServicesEpicCommon/Private/Online/EOSAuthLoginOptionsCommon.cpp
@@ -93,7 +93,7 @@ FEOSAuthLoginOptionsCommon::FEOSAuthLoginOptionsCommon()
 {
 	// EOS_Auth_LoginOptions init
 	UE_EOS_CHECK_API_MISMATCH(EOS_AUTH_LOGIN_API_LATEST, 3);
-	ApiVersion = 2;
+	ApiVersion = 3;
 	Credentials = &CredentialsData;
 	ScopeFlags = EOS_EAuthScopeFlags::EOS_AS_NoFlags;
 	LoginFlags = 0;
